- Added a new style brush to concert style using the 40x40 image specifying it is a 20x20 one (automatically resized by slate) used for the concert icons.
  - Added the 'small' brush style alongside the 'normal' one when creating the slate icon for the tool bar.

diff --git a/Engine/Plugins/Developer/Concert/ConcertFrontend/Source/ConcertFrontend/Private/ConcertFrontendModule.cpp b/Engine/Plugins/Developer/Concert/ConcertFrontend/Source/ConcertFrontend/Private/ConcertFrontendModule.cpp
--- a/Engine/Plugins/Developer/Concert/ConcertFrontend/Source/ConcertFrontend/Private/ConcertFrontendModule.cpp
+++ b/Engine/Plugins/Developer/Concert/ConcertFrontend/Source/ConcertFrontend/Private/ConcertFrontendModule.cpp
@@ -309,9 +309,9 @@ private:
 		IConcertClientPtr ConcertClient = IConcertModule::Get().GetClientInstance();
 		if (ConcertClient.IsValid() && ConcertClient->GetSessionConnectionStatus() == EConcertConnectionStatus::Connected)
 		{
-			return FSlateIcon(FConcertFrontendStyle::GetStyleSetName(), "Concert.Online");
+			return FSlateIcon(FConcertFrontendStyle::GetStyleSetName(), "Concert.Online", "Concert.Online.Small");
 		}
-		return FSlateIcon(FConcertFrontendStyle::GetStyleSetName(), "Concert.Offline");
+		return FSlateIcon(FConcertFrontendStyle::GetStyleSetName(), "Concert.Offline", "Concert.Offline.Small");
 	}
 
 	/**
diff --git a/Engine/Plugins/Developer/Concert/ConcertFrontend/Source/ConcertFrontend/Private/ConcertFrontendStyle.cpp b/Engine/Plugins/Developer/Concert/ConcertFrontend/Source/ConcertFrontend/Private/ConcertFrontendStyle.cpp
--- a/Engine/Plugins/Developer/Concert/ConcertFrontend/Source/ConcertFrontend/Private/ConcertFrontendStyle.cpp
+++ b/Engine/Plugins/Developer/Concert/ConcertFrontend/Source/ConcertFrontend/Private/ConcertFrontendStyle.cpp
@@ -37,6 +37,7 @@ void FConcertFrontendStyle::Initialize()
 
 	// Const icon sizes
 	const FVector2D Icon16x16(16.0f, 16.0f);
+	const FVector2D Icon20x20(20.0f, 20.0f);
 	const FVector2D Icon32x32(32.0f, 32.0f);
 	const FVector2D Icon40x40(40.0f, 40.0f);
 
@@ -47,6 +48,10 @@ void FConcertFrontendStyle::Initialize()
 	StyleSet->Set("Concert.OtherLock", new IMAGE_PLUGIN_BRUSH("Icons/icon_ConcertOtherLock_16x", Icon16x16));
 	StyleSet->Set("Concert.ModifiedByOther", new IMAGE_PLUGIN_BRUSH("Icons/icon_ConcertModifiedByOther_16x", Icon16x16));
 
+	// 20x20 -> For toolbar small icons.
+	StyleSet->Set("Concert.Online.Small", new IMAGE_PLUGIN_BRUSH("Icons/icon_ConcertOnline_40x", Icon20x20));
+	StyleSet->Set("Concert.Offline.Small", new IMAGE_PLUGIN_BRUSH("Icons/icon_ConcertOffline_40x", Icon20x20));
+
 	// 40x40
 	StyleSet->Set("Concert.Online", new IMAGE_PLUGIN_BRUSH("Icons/icon_ConcertOnline_40x", Icon40x40));
 	StyleSet->Set("Concert.Offline", new IMAGE_PLUGIN_BRUSH("Icons/icon_ConcertOffline_40x", Icon40x40));
