- Fixed a multi-users bug in the active session panel (the one displaying the connected user and the session history) not displaying all users connected to the session.
- Removed dead code from SActiveSession.h. Removed member functions declared, but not unimplemented.
- Added support to the active session panel to update the local client session info (display name/avatar color) if they changes in real time.

diff --git a/Engine/Plugins/Developer/Concert/ConcertFrontend/Source/ConcertFrontend/Private/Widgets/SActiveSession.cpp b/Engine/Plugins/Developer/Concert/ConcertFrontend/Source/ConcertFrontend/Private/Widgets/SActiveSession.cpp
--- a/Engine/Plugins/Developer/Concert/ConcertFrontend/Source/ConcertFrontend/Private/Widgets/SActiveSession.cpp
+++ b/Engine/Plugins/Developer/Concert/ConcertFrontend/Source/ConcertFrontend/Private/Widgets/SActiveSession.cpp
@@ -209,7 +209,7 @@ void SActiveSession::Construct(const FArguments& InArgs, const TSharedRef<SDockT
 		if (ClientSession.IsValid())
 		{
 			WeakSessionPtr = ClientSession;
-			ClientInfo = FConcertSessionClientInfo({ ClientSession->GetSessionClientEndpointId(), ClientSession->GetLocalClientInfo() });
+			ClientInfo = MakeShared<FConcertSessionClientInfo>(FConcertSessionClientInfo{ClientSession->GetSessionClientEndpointId(), ClientSession->GetLocalClientInfo()});
 			SessionClientChangedHandle = ClientSession->OnSessionClientChanged().AddSP(this, &SActiveSession::HandleSessionClientChanged);
 		}
 	}
@@ -369,7 +369,11 @@ void SActiveSession::Construct(const FArguments& InArgs, const TSharedRef<SDockT
 		]
 	];
 
-	ReloadClients();
+	// Create a timer to periodically poll the this client(s) info to detect if he changed its display name or avatar color because
+	// IConcertClientsession::OnSessionClientChanged() doesn't trigger when the 'local' client changes. This needs to be polled.
+	RegisterActiveTimer(1.0f, FWidgetActiveTimerDelegate::CreateSP(this, &SActiveSession::HandleLocalClientInfoChangePollingTimer));
+
+	UpdateSessionClientListView();
 }
 
 TSharedRef<ITableRow> SActiveSession::HandleGenerateRow(TSharedPtr<FConcertSessionClientInfo> InClientInfo, const TSharedRef<STableViewBase>& OwnerTable) const
@@ -383,9 +387,9 @@ void SActiveSession::HandleSessionStartup(TSharedRef<IConcertClientSession> InCl
 	WeakSessionPtr = InClientSession;
 	SessionClientChangedHandle = InClientSession->OnSessionClientChanged().AddSP(this, &SActiveSession::HandleSessionClientChanged);
 
-	ClientInfo = FConcertSessionClientInfo({ InClientSession->GetSessionClientEndpointId(), InClientSession->GetLocalClientInfo() });
+	ClientInfo = MakeShared<FConcertSessionClientInfo>(FConcertSessionClientInfo{InClientSession->GetSessionClientEndpointId(), InClientSession->GetLocalClientInfo()});
 
-	ReloadClients();
+	UpdateSessionClientListView();
 
 	if (SessionHistory.IsValid())
 	{
@@ -415,74 +419,121 @@ void SActiveSession::HandleSessionShutdown(TSharedRef<IConcertClientSession> InC
 
 void SActiveSession::HandleSessionClientChanged(IConcertClientSession&, EConcertClientStatus ClientStatus, const FConcertSessionClientInfo& InClientInfo)
 {
-	if (ClientStatus == EConcertClientStatus::Connected || ClientStatus == EConcertClientStatus::Updated)
-	{
-		UpdateAvailableClients({ MakeShared<FConcertSessionClientInfo>(InClientInfo) });
-	}
-	else
-	{
-		int32 ClientIndex = Clients.IndexOfByPredicate([&InClientInfo](const TSharedPtr<FConcertSessionClientInfo>& PotentialClient)
-		{
-			return PotentialClient->ClientEndpointId == InClientInfo.ClientEndpointId;
-		});
+	// Update the view for a specific client.
+	UpdateSessionClientListView(&InClientInfo, ClientStatus);
+}
 
-		if (ClientIndex != INDEX_NONE)
+EActiveTimerReturnType SActiveSession::HandleLocalClientInfoChangePollingTimer(double InCurrentTime, float InDeltaTime)
+{
+	// NOTE: As Jan 2019, the client info never updates in real time, so the code below will not be useful until this
+	//       this feature gets implemented.
+
+	TSharedPtr<IConcertClientSession> Session = WeakSessionPtr.Pin();
+	if (Session.IsValid() && ClientInfo.IsValid())
+	{
+		// Check if the local client info cached as class member is out dated with respect to the one held by the session changes. Just check the info displayed by this panel.
+		const FConcertClientInfo& LatestClientInfo = Session->GetLocalClientInfo();
+		if (LatestClientInfo.DisplayName != ClientInfo->ClientInfo.DisplayName || LatestClientInfo.AvatarColor != ClientInfo->ClientInfo.AvatarColor)
 		{
-			RefreshListViewAfterFunction([this, ClientIndex]() 
-			{
-				Clients.RemoveAt(ClientIndex);
-			});
+			// Update the view for this client info.
+			ClientInfo->ClientInfo = LatestClientInfo;
+			UpdateSessionClientListView(ClientInfo.Get(), EConcertClientStatus::Updated);
 		}
 	}
+
+	return EActiveTimerReturnType::Continue;
 }
 
-void SActiveSession::UpdateAvailableClients(TArray<TSharedPtr<FConcertSessionClientInfo>>&& InAvailableClients)
+void SActiveSession::UpdateSessionClientListView(const FConcertSessionClientInfo* InClientInfo, EConcertClientStatus ClientStatus)
 {
-	RefreshListViewAfterFunction([&Clients = this->Clients, &ClientInfo = this->ClientInfo, AvailableClients = MoveTemp(InAvailableClients)]() mutable
+	// NOTE: Calling 'Session->GetSessionClients()' while handling IConcertClientSession::OnSessionClientChanged() event may not
+	//       return the up-to-date list as one would expect. When a client connects, the client session implementation adds the
+	//       client to its list before broadcasting the notification, but when a client disconnects, it removes it from the list
+	//       after the broadcast, so in the callback, we would read the out of date list on disconnect. This may change in the
+	//       future, but to mitigate that, this function has one code path to deal with single client change received from
+	//       IConcertClientSession::OnSessionClientChanged() and one code path to initialize the view. While this need more code,
+	//       it is also more efficient.
+
+	// We expect the UI to be constructed in Construct() function, prior this function gets called.
+	check(ClientsListView.IsValid());
+
+	// Try pinning the session.
+	TSharedPtr<IConcertClientSession> Session = WeakSessionPtr.Pin();
+	if (Session.IsValid())
 	{
-		ConcertFrontendUtils::SyncArraysByPredicate(Clients, MoveTemp(AvailableClients), [](const TSharedPtr<FConcertSessionClientInfo>& ClientToFind)
+		// Remember the element selected in the list view to reselect it later.
+		TArray<TSharedPtr<FConcertSessionClientInfo>> SelectedItems = ClientsListView->GetSelectedItems();
+		checkf(SelectedItems.Num() <= 1, TEXT("ActiveSession's client list view should not support multiple selection."));
+
+		// If this is about a specific client update. (i.e. responding to a IConcertClientSession::OnSessionClientChanged() event)
+		if (InClientInfo)
 		{
-			return [ClientToFind](const TSharedPtr<FConcertSessionClientInfo>& PotentialClient)
+			if (ClientStatus == EConcertClientStatus::Connected)
 			{
-				return PotentialClient->ClientEndpointId == ClientToFind->ClientEndpointId;
-			};
-		});
-
-		if (ClientInfo.IsSet())
+				Clients.Emplace(MakeShared<FConcertSessionClientInfo>(*InClientInfo));
+			}
+			else
+			{
+				int32 Index = Clients.IndexOfByPredicate([InClientInfo](const TSharedPtr<FConcertSessionClientInfo>& Visited) { return InClientInfo->ClientEndpointId == Visited->ClientEndpointId; });
+				if (Index != INDEX_NONE)
+				{
+					if (ClientStatus == EConcertClientStatus::Disconnected)
+					{
+						Clients.RemoveAt(Index); // We want to preserve items relative order.
+					}
+					else if (ensure(ClientStatus == EConcertClientStatus::Updated)) // Ensure we handled all status in EConcertClientStatus.
+					{
+						*Clients[Index] = *InClientInfo; // Update the client info.
+					}
+				}
+			}
+		}
+		else // Sync everything.
 		{
-			Clients.Emplace(MakeShared<FConcertSessionClientInfo>(ClientInfo.GetValue()));
+			// Convert the list of clients to a list of shared pointer to client (the list view model asks for shared pointers).
+			TArray<FConcertSessionClientInfo> OtherConnectedClients = Session->GetSessionClients();
+			TArray<TSharedPtr<FConcertSessionClientInfo>> UpdatedClientList;
+			UpdatedClientList.Reserve(OtherConnectedClients.Num() + 1); // +1 is to include this local client (see below).
+			Algo::Transform(OtherConnectedClients, UpdatedClientList, [](FConcertSessionClientInfo InClient)
+			{
+				return MakeShared<FConcertSessionClientInfo>(MoveTemp(InClient));
+			});
+
+			// Add this local client as it is not part of the list returned by GetSessionClients(). (The client connected from this process).
+			if (ClientInfo.IsValid())
+			{
+				UpdatedClientList.Emplace(ClientInfo);
+			}
+
+			// Merge the list used by the list view (the model) with the updated list, removing clients who left and adding the one who joined.
+			ConcertFrontendUtils::SyncArraysByPredicate(Clients, MoveTemp(UpdatedClientList), [](const TSharedPtr<FConcertSessionClientInfo>& ClientToFind)
+			{
+				return [ClientToFind](const TSharedPtr<FConcertSessionClientInfo>& PotentialClient)
+				{
+					return PotentialClient->ClientEndpointId == ClientToFind->ClientEndpointId;
+				};
+			});
 		}
 
+		// Sort the list by display name alphabetically.
 		Clients.StableSort([](const TSharedPtr<FConcertSessionClientInfo>& ClientOne, const TSharedPtr<FConcertSessionClientInfo>& ClientTwo)
 		{
 			return ClientOne->ClientInfo.DisplayName < ClientTwo->ClientInfo.DisplayName;
 		});
-	});
-}
-
-void SActiveSession::RefreshListViewAfterFunction(TFunctionRef<void()> InFunction)
-{
-	if (ensureMsgf(ClientsListView.IsValid(), TEXT("RefreshListViewAfterFunction should not be called with an invalid clients ListView.")))
-	{
-		TSharedPtr<FConcertSessionClientInfo> SelectedClient;
-		TArray<TSharedPtr<FConcertSessionClientInfo>> SelectedClients = ClientsListView->GetSelectedItems();
-
-		checkf(SelectedClients.Num() <= 1, TEXT("ActiveSession's client list view should not support multiple selection."));
 
-		if (SelectedClients.Num() > 0)
+		// If a client row was selected, select it back (if still available).
+		if (SelectedItems.Num() > 0)
 		{
-			SelectedClient = SelectedClients[0];
-		}
-
-		InFunction();
-
-		ClientsListView->RequestListRefresh();
-
-		if (SelectedClient.IsValid())
-		{
-			SetSelectedClient(SelectedClient->ClientEndpointId);
+			ClientsListView->SetSelection(SelectedItems[0]);
 		}
 	}
+	else // The session appears to be invalid.
+	{
+		// Clear the list of clients.
+		Clients.Reset();
+	}
+
+	ClientsListView->RequestListRefresh();
 }
 
 const FButtonStyle& SActiveSession::GetConnectionIconStyle() const
@@ -645,23 +696,4 @@ TSharedPtr<FConcertSessionClientInfo> SActiveSession::FindAvailableClient(const
 	return FoundClientPtr ? *FoundClientPtr : nullptr;
 }
 
-void SActiveSession::ReloadClients()
-{
-	TSharedPtr<IConcertClientSession> ClientSession = WeakSessionPtr.Pin();
-	TArray<TSharedPtr<FConcertSessionClientInfo>> ClientPtrs;
-
-	if (ClientSession.IsValid())
-	{
-		TArray<FConcertSessionClientInfo> AvailableClients = ClientSession->GetSessionClients();
-
-		ClientPtrs.Reserve(AvailableClients.Num() + 1);
-		Algo::Transform(AvailableClients, ClientPtrs, [](FConcertSessionClientInfo InClient)
-		{
-			return MakeShared<FConcertSessionClientInfo>(InClient);
-		});
-	}
-
-	UpdateAvailableClients(MoveTemp(ClientPtrs));
-}
-
 #undef LOCTEXT_NAMESPACE /* SActiveSession */
diff --git a/Engine/Plugins/Developer/Concert/ConcertFrontend/Source/ConcertFrontend/Private/Widgets/SActiveSession.h b/Engine/Plugins/Developer/Concert/ConcertFrontend/Source/ConcertFrontend/Private/Widgets/SActiveSession.h
--- a/Engine/Plugins/Developer/Concert/ConcertFrontend/Source/ConcertFrontend/Private/Widgets/SActiveSession.h
+++ b/Engine/Plugins/Developer/Concert/ConcertFrontend/Source/ConcertFrontend/Private/Widgets/SActiveSession.h
@@ -46,10 +46,11 @@ private:
 	/** Handle a session client change */
 	void HandleSessionClientChanged(IConcertClientSession&, EConcertClientStatus ClientStatus, const FConcertSessionClientInfo& ClientInfo);
 
-	FString HandleSessionName() const;
+	/** Polls the local client info and detect if it changed in order to update its representation in real time.*/
+	EActiveTimerReturnType HandleLocalClientInfoChangePollingTimer(double InCurrentTime, float InDeltaTime);
 
 	/** Update the list of clients while keeping the alphabetical sorting */
-	void UpdateAvailableClients(TArray<TSharedPtr<FConcertSessionClientInfo>>&& InAvailableClients);
+	void UpdateSessionClientListView(const FConcertSessionClientInfo* Client = nullptr, EConcertClientStatus Status = EConcertClientStatus::Updated);
 
 	/** Set the selected client in the clients list view*/
 	void SetSelectedClient(const FGuid& InClientEndpointId, ESelectInfo::Type SelectInfo = ESelectInfo::Direct);
@@ -57,12 +58,6 @@ private:
 	/** Find a client with its endpoint id */
 	TSharedPtr<FConcertSessionClientInfo> FindAvailableClient(const FGuid& InClientEndpointId) const;
 	
-	/** Sync the clients with the clients from the current session */
-	void ReloadClients();
-
-	/** Refresh the clients list after a function call while preserving alphabetical order and current selection */
-	void RefreshListViewAfterFunction(TFunctionRef<void()> InFunction);
-
 	/** Handling for the status icon and text */
 	const FButtonStyle& GetConnectionIconStyle() const;
 	FSlateColor GetConnectionIconColor() const;
@@ -77,9 +72,6 @@ private:
 	FReply OnClickResumeSession();
 	FReply OnClickLeaveSession();
 
-	/** Handling for spawning the Active Session tabs. */
-	TSharedRef<SDockTab> HandleTabManagerSpawnTab(const class FSpawnTabArgs&, const FName TabIdentifier);
-
 private:
 	/** Holds a concert client session. */
 	TWeakPtr<IConcertClientSession> WeakSessionPtr;
@@ -91,7 +83,7 @@ private:
 	TArray<TSharedPtr<FConcertSessionClientInfo>> Clients;
 
 	/** Information about the machine's client. */
-	TOptional<FConcertSessionClientInfo> ClientInfo;
+	TSharedPtr<FConcertSessionClientInfo> ClientInfo;
 
 	/** Delegate handle for session clients state changes. */
 	FDelegateHandle SessionClientChangedHandle;
